Linux: fixed GDB pretty printer for FName

diff --git a/Engine/Extras/GDBPrinters/UE4Printers.py b/Engine/Extras/GDBPrinters/UE4Printers.py
--- a/Engine/Extras/GDBPrinters/UE4Printers.py
+++ b/Engine/Extras/GDBPrinters/UE4Printers.py
@@ -505,10 +505,7 @@ class FNamePrinter:
 		if self.Value.is_optimized_out:
 			return '<optimized out>'
 
-		DisplayIndex = self.Value['DisplayIndex']
-
-		if not DisplayIndex:
-			DisplayIndex = self.Value['ComparisonIndex']
+		DisplayIndex = self.Value['ComparisonIndex']
 
 		if DisplayIndex >= 1048576:
 			return 'invalid'
@@ -673,9 +670,9 @@ def lookup_function (val):
 
 def build_dictionary ():
 	pretty_printers_dict[re.compile('^FString$')] = lambda typename, val: FStringPrinter(typename, val)
-#	pretty_printers_dict[re.compile('^FNameEntry$')] = lambda typename, val: FNameEntryPrinter(typename, val)
-#	pretty_printers_dict[re.compile('^FName$')] = lambda typename, val: FNamePrinter(typename, val)
-#	pretty_printers_dict[re.compile('^FMinimalName$')] = lambda typename, val: FMinimalNamePrinter(typename, val)
+	pretty_printers_dict[re.compile('^FNameEntry$')] = lambda typename, val: FNameEntryPrinter(typename, val)
+	pretty_printers_dict[re.compile('^FName$')] = lambda typename, val: FNamePrinter(typename, val)
+	pretty_printers_dict[re.compile('^FMinimalName$')] = lambda typename, val: FMinimalNamePrinter(typename, val)
 	pretty_printers_dict[re.compile('^TArray<.+,.+>$')] = lambda typename, val: TArrayPrinter(typename, val)
 	pretty_printers_dict[re.compile('^TBitArray<.+>$')] = lambda typename, val: TBitArrayPrinter(typename, val)
 	pretty_printers_dict[re.compile('^TChunkedArray<.+>$')] = lambda typename, val: TChunkedArrayPrinter(typename, val)
