Fixed TSparseArray printer.
Fixed FName printer so that it doesn't use a hardcoded value to determine if it's a valid FName or not.

Disabled broken TMap and TSet printers, but you will see output from the TSparseArray printer so you can see the stored values.

diff --git a/Engine/Extras/GDBPrinters/UEPrinters.py b/Engine/Extras/GDBPrinters/UEPrinters.py
--- a/Engine/Extras/GDBPrinters/UEPrinters.py
+++ b/Engine/Extras/GDBPrinters/UEPrinters.py
@@ -205,12 +205,11 @@ class TSparseArrayPrinter:
 			self.Counter = -1
 			self.Typename = typename
 			self.ElementType = self.Value.type.template_argument(0)
-			self.ElementTypeSize = self.ElementType.sizeof
-			assert self.ElementTypeSize != 0
 
 			try:
 				self.NumFreeIndices = self.Value['NumFreeIndices']
 				self.Data = self.Value['Data']
+				self.InternalElementType = self.Data.type.template_argument(0)
 				self.ArrayNum = self.Data['ArrayNum']
 				if self.ArrayNum.is_optimized_out:
 					self.ArrayNum = 0
@@ -246,14 +245,12 @@ class TSparseArrayPrinter:
 				else:
 					Data = (self.InlineData.address.cast(gdb.lookup_type("int").pointer())[self.Counter/32] >> self.Counter) & 1
 
-				Value = None
 				if Data != 0:
-					offset = self.Counter * self.ElementTypeSize
-					Value = (self.AllocatorData + offset).cast(self.ElementType.pointer())
+					ElementOrFreeListValue = (self.AllocatorData.cast(self.InternalElementType.pointer()) + self.Counter).dereference()
+					Value = ElementOrFreeListValue['ElementData'].reinterpret_cast(self.ElementType.reference())
+					return ('[%d]' % self.Counter, Value)
 				else:
-					Value = None
-
-				return ('[%s]' % self.Counter, Value.dereference())
+					return self.__next__()
 
 	def __init__(self, typename, val):
 		self.Value = val
@@ -271,7 +268,7 @@ class TSparseArrayPrinter:
 		return self._iterator(self.Value, self.Typename)
 
 	def display_hint(self):
-		return 'array'
+		return 'string'
 
 # ------------------------------------------------------------------------------
 # TSet
@@ -541,7 +538,8 @@ class FNamePrinter:
 		Index = self.Value['ComparisonIndex']['Value']
 		IndexValue = int(Index)
 
-		if IndexValue >= 4194304:
+		UnusedMask = gdb.parse_and_eval('FNameDebugVisualizer::UnusedMask')
+		if (IndexValue & UnusedMask) != 0:
 			return 'invalid'
 		else:
 			Expr = '((FNameEntry&)GNameBlocksDebug['+str(IndexValue)+' >> FNameDebugVisualizer::OffsetBits][FNameDebugVisualizer::EntryStride * ('+str(IndexValue)+' & FNameDebugVisualizer::OffsetMask)])'
@@ -716,9 +714,9 @@ def build_dictionary ():
 	pretty_printers_dict[re.compile('^TChunkedArray<.+>$')] = lambda typename, val: TChunkedArrayPrinter(typename, val)
 	pretty_printers_dict[re.compile('^TSparseArray<.+>$')] = lambda typename, val: TSparseArrayPrinter(typename, val)
 	pretty_printers_dict[re.compile('^TSetElement<.+>$')] = lambda typename, val: TSetElementPrinter(typename, val)
-	pretty_printers_dict[re.compile('^TSet<.+>$')] = lambda typename, val: TSetPrinter(typename, val)
+#	pretty_printers_dict[re.compile('^TSet<.+>$')] = lambda typename, val: TSetPrinter(typename, val)
 	pretty_printers_dict[re.compile('^FBitReference$')] = lambda typename, val: FBitReferencePrinter(typename, val)
-	pretty_printers_dict[re.compile('^TMap<.+,.+,.+>$')] = lambda typename, val: TMapPrinter(typename, val)
+#	pretty_printers_dict[re.compile('^TMap<.+,.+,.+>$')] = lambda typename, val: TMapPrinter(typename, val)
 	pretty_printers_dict[re.compile('^TPair<.+,.+>$')] = lambda typename, val: TTuplePrinter(typename, val)
 	pretty_printers_dict[re.compile('^TTuple<.+,.+>$')] = lambda typename, val: TTuplePrinter(typename, val)
 	pretty_printers_dict[re.compile('^TWeakObjectPtr<.+>$')] = lambda typename, val: TWeakObjectPtrPrinter(typename, val)
