Fixed TArray pretty printer for GDB.

diff --git a/Engine/Extras/GDBPrinters/UEPrinters.py b/Engine/Extras/GDBPrinters/UEPrinters.py
--- a/Engine/Extras/GDBPrinters/UEPrinters.py
+++ b/Engine/Extras/GDBPrinters/UEPrinters.py
@@ -646,14 +646,19 @@ class TArrayPrinter:
 
 				if self.ArrayNum > 0:
 					self.AllocatorInstance = self.Value['AllocatorInstance']
-					self.AllocatorInstanceData = self.AllocatorInstance['Data']
+					self.AllocatorInstanceData = None
+					self.InlineData = None
+					self.SecondaryDataData = None
+					try:
+						self.AllocatorInstanceData = self.AllocatorInstance['Data']
+					except:
+						pass
 					try:
 						self.InlineData = self.AllocatorInstance['InlineData']
-						self.SecondaryData = self.AllocatorInstance['SecondaryData']
-						if self.SecondaryData != None:
-							self.SecondaryDataData = self.SecondaryData['Data']
-						else:
-							self.SecondaryDataData = None
+					except:
+						pass
+					try:
+						self.SecondaryDataData = self.AllocatorInstance['SecondaryData']['Data']
 					except:
 						pass
 			except:
@@ -672,14 +677,14 @@ class TArrayPrinter:
 				raise StopIteration
 
 			try:
-				if self.AllocatorInstanceData != None:
-					data = self.AllocatorInstanceData.cast(self.TType.pointer())
-				elif self.SecondaryDataDataVal > 0:
-					data = self.SecondaryDataData.cast(self.TType.pointer())
-				else:
-					data = self.InlineData.cast(self.TType.pointer())
+			    if self.AllocatorInstanceData != None:
+				    data = self.AllocatorInstanceData.cast(self.TType.pointer())
+			    elif self.SecondaryDataData != None:
+				    data = self.SecondaryDataData.cast(self.TType.pointer())
+			    else:
+				    data = self.InlineData.cast(self.TType.pointer())
 			except:
-				return ('[%d]' % self.Counter, "optmized")
+				return ('[%d]' % self.Counter, "optimized")
 
 			return ('[%d]' % self.Counter, data[self.Counter])
 
