Fixed TChunkedArray pretty printer for GDB.

diff --git a/Engine/Extras/GDBPrinters/UEPrinters.py b/Engine/Extras/GDBPrinters/UEPrinters.py
--- a/Engine/Extras/GDBPrinters/UEPrinters.py
+++ b/Engine/Extras/GDBPrinters/UEPrinters.py
@@ -133,12 +133,12 @@ class TChunkedArrayPrinter:
 	"Print TChunkedArray"
 
 	class _iterator(Iterator):
-		def __init__(self, val, typename):
+		def __init__(self, val):
 			self.Value = val
-			self.Typename = typename
 			self.Counter = -1
 			self.ElementType = self.Value.type.template_argument(0)
 			self.ElementTypeSize = self.ElementType.sizeof
+			self.NumElementsPerChunk = self.Value.type.template_argument(1)/self.ElementTypeSize
 
 			try:
 				self.NumElements = self.Value['NumElements']
@@ -168,17 +168,12 @@ class TChunkedArrayPrinter:
 			if self.Counter >= self.NumElements:
 				raise StopIteration()
 
-			Expr = '(unsigned)sizeof('+str(self.Typename)+'::FChunk)/'+str(self.ElementTypeSize)
-			self.ChunkBytes = gdb.parse_and_eval(Expr)
-			assert self.ChunkBytes != 0
-
-			Expr = '*(*((('+str(self.ElementType.name)+'**)'+str(self.AllocatorData)+')+'+str(self.Counter / self.ChunkBytes)+')+'+str(self.Counter % self.ChunkBytes)+')'
+			Expr = '*(*((('+str(self.ElementType.name)+'**)'+str(self.AllocatorData)+')+'+str(self.Counter / self.NumElementsPerChunk)+')+'+str(self.Counter % self.NumElementsPerChunk)+')'
 			Val = gdb.parse_and_eval(Expr)
 			return ('[%d]' % self.Counter, Val)
 
 	def __init__(self, val):
 		self.Value = val
-		self.Typename = typename
 		self.NumElements = self.Value['NumElements']
 
 	def to_string(self):
@@ -189,7 +184,7 @@ class TChunkedArrayPrinter:
 		pass
 
 	def children(self):
-		return self._iterator(self.Value, self.Typename)
+		return self._iterator(self.Value)
 
 	def display_hint(self):
 		return 'array'
