Fixed GDB pretty printing of TSet as well as some minor fixes.

diff --git a/Engine/Extras/GDBPrinters/UEPrinters.py b/Engine/Extras/GDBPrinters/UEPrinters.py
--- a/Engine/Extras/GDBPrinters/UEPrinters.py
+++ b/Engine/Extras/GDBPrinters/UEPrinters.py
@@ -80,7 +80,7 @@ class TBitArrayPrinter:
 					self.AllocatorInstance = self.Value['AllocatorInstance']
 					self.InlineData = self.AllocatorInstance['InlineData']
 					self.SecondaryData = self.AllocatorInstance['SecondaryData']
-					self.SecondaryDataData = self.AllocatorInstance['SecondaryData']['Data']
+					self.SecondaryDataData = None
 					if self.SecondaryData != None:
 						self.SecondaryDataData = self.SecondaryData['Data']
 			except:
@@ -196,11 +196,10 @@ class TSparseArrayPrinter:
 	"Print TSparseArray"
 
 	class _iterator(Iterator):
-		def __init__(self, val, typename):
+		def __init__(self, val):
 
 			self.Value = val
 			self.Counter = -1
-			self.Typename = typename
 			self.ElementType = self.Value.type.template_argument(0)
 
 			try:
@@ -217,7 +216,9 @@ class TSparseArrayPrinter:
 					self.AllocationFlagsInstance = self.AllocationFlags['AllocatorInstance']
 					self.InlineData = self.AllocationFlagsInstance['InlineData']
 					self.SecondaryData = self.AllocationFlagsInstance['SecondaryData']
-					self.SecondaryDataData = self.SecondaryData['Data']
+					self.SecondaryDataData = None
+					if self.SecondaryData != None:
+						self.SecondaryData['Data']
 			except:
 				raise
 
@@ -251,7 +252,6 @@ class TSparseArrayPrinter:
 
 	def __init__(self, val):
 		self.Value = val
-		self.Typename = typename
 		self.ArrayNum = self.Value['Data']['ArrayNum']
 
 	def to_string(self):
@@ -262,7 +262,7 @@ class TSparseArrayPrinter:
 		pass
 
 	def children(self):
-		return self._iterator(self.Value, self.Typename)
+		return self._iterator(self.Value)
 
 	def display_hint(self):
 		return 'string'
@@ -274,33 +274,40 @@ class TSetPrinter:
 	"Print TSet"
 
 	class _iterator(Iterator):
-		def __init__(self, val, typename):
+		def __init__(self, val):
 			self.Value = val
 			self.Counter = -1
-			self.typename = typename
 			self.ElementType = self.Value.type.template_argument(0)
+			self.TSetElementType = gdb.lookup_type("TSetElement<" + self.ElementType.name + ">")
+			self.ElementsArrayNum = 0
 
 			try:
 				self.Elements = self.Value["Elements"]
 				self.ElementsData = self.Elements["Data"]
 				self.ElementsArrayNum = self.ElementsData['ArrayNum']
 				self.NumFreeIndices = self.Elements['NumFreeIndices']
-
 				self.AllocatorInstance = self.ElementsData['AllocatorInstance']
-				self.AllocatorInstanceData = self.AllocatorInstance['Data']
-
 				self.AllocationFlags = self.Elements['AllocationFlags']
-				self.AllocationFlagsInstance = self.AllocationFlags['AllocatorInstance']
-				self.InlineData = self.AllocationFlagsInstance['InlineData']
-				self.SecondaryData = self.AllocationFlagsInstance['SecondaryData']
-				self.SecondaryDataData = self.SecondaryData['Data']
+				self.AllocatorInstanceData = None
+				self.InlineData = None
+				self.SecondaryDataData = None
+				try:
+					self.AllocatorInstanceData = self.AllocatorInstance['Data']
+				except:
+					pass
+				try:
+					self.InlineData = self.AllocationFlags['AllocatorInstance']['InlineData']
+				except:
+					pass
+				try:
+					SecondaryData = self.AllocationFlags['AllocatorInstance']['SecondaryData']
+					if SecondaryData != None:
+						self.SecondaryDataData = SecondaryData['Data']
+				except:
+					pass
 			except:
 				self.ElementsArrayNum = 0
 
-			Expr = '(size_t)sizeof(FSetElementId) + sizeof(int32)'
-			TSetElement = gdb.parse_and_eval(Expr)
-			self.ElementTypeSize = self.ElementType.sizeof + TSetElement
-
 		def __iter__(self):
 			return self
 
@@ -310,24 +317,20 @@ class TSetPrinter:
 			if self.Counter >= self.ElementsArrayNum:
 				raise StopIteration()
 			else:
-				Data = None
+				IsUsed = 0
 				if self.SecondaryDataData > 0:
-					Data = (self.SecondaryDataData.address.cast(gdb.lookup_type("int").pointer())[self.Counter/32] >> self.Counter) & 1
+					IsUsed = (self.SecondaryDataData.address.cast(gdb.lookup_type("uint32").pointer())[self.Counter/32] >> self.Counter) & 1
 				else:
-					Data = (self.InlineData.address.cast(gdb.lookup_type("int").pointer())[self.Counter/32] >> self.Counter) & 1
+					IsUsed = (self.InlineData.address.cast(gdb.lookup_type("uint32").pointer())[self.Counter/32] >> self.Counter) & 1
 
-				Value = None
-				if Data != 0:
-					offset = self.Counter * self.ElementTypeSize
-					Value = (self.AllocatorInstanceData + offset).cast(self.ElementType.pointer())
+				if IsUsed == 1:
+					Value = self.AllocatorInstanceData.cast(self.TSetElementType.pointer())[self.Counter]
+					return ('[%d]' % self.Counter, Value)
 				else:
-					Value = None
-
-			return ('[%s]' % self.Counter, Value.dereference())
+					return self.__next__()
 
 	def __init__(self, val):
 		self.Value = val
-		self.typename = typename
 		self.ArrayNum = self.Value["Elements"]["Data"]['ArrayNum']
 
 	def to_string(self):
@@ -338,7 +341,7 @@ class TSetPrinter:
 		pass
 
 	def children(self):
-		return self._iterator(self.Value, self.typename)
+		return self._iterator(self.Value)
 
 	def display_hint(self):
 		return 'array'
@@ -653,7 +656,9 @@ class TArrayPrinter:
 					except:
 						pass
 					try:
-						self.SecondaryDataData = self.AllocatorInstance['SecondaryData']['Data']
+						SecondaryData = self.AllocatorInstance['SecondaryData']
+						if SecondaryData != None:
+							self.SecondaryDataData = SecondaryData['Data']
 					except:
 						pass
 			except:
@@ -721,7 +726,7 @@ def build_ue_pretty_printer():
 	pp.add_printer("TChunkedArray", '^TChunkedArray<.+>$', TChunkedArrayPrinter)
 	pp.add_printer("TSparseArray", '^TSparseArray<.+>$', TSparseArrayPrinter)
 	pp.add_printer("TSetElement", '^TSetElement<.+>$', TSetElementPrinter)
-#	pp.add_printer("TSet", '^TSet<.+>$', TSetPrinter)
+	pp.add_printer("TSet", '^TSet<.+>$', TSetPrinter)
 	pp.add_printer("FBitReference", '^FBitReference$', FBitReferencePrinter)
 #	pp.add_printer("TMap", '^TMap<.+,.+,.+>$', TMapPrinter)
 	pp.add_printer("TPair", '^TPair<.+,.+>$', TTuplePrinter)
