Fixed TMapPrinter not working with latest WSL

#Synced-CL 45873191

diff --git a/Engine/Extras/GDBPrinters/UEPrinters.py b/Engine/Extras/GDBPrinters/UEPrinters.py
--- a/Engine/Extras/GDBPrinters/UEPrinters.py
+++ b/Engine/Extras/GDBPrinters/UEPrinters.py
@@ -320,7 +320,7 @@ class TMapPrinter(ContainerPrinterHelper):
 
 	def children(self):
 		for Str, Pair in super().children():
-			yield (f'[{Pair['Key']}] = {Pair['Value'].type}', Pair['Value'])
+			yield (f"[{Pair['Key']}] = {Pair['Value'].type}", Pair['Value'])
 
 # ------------------------------------------------------------------------------
 # TWeakObjectPtr
